Merging CL#9483137

#PreLoadScreenMoviePlayer CIS fix

diff --git a/Engine/Plugins/Runtime/PreLoadScreenMoviePlayer/Source/PreLoadScreenMoviePlayer/Private/PreLoadMoviePlayerScreenBase.cpp b/Engine/Plugins/Runtime/PreLoadScreenMoviePlayer/Source/PreLoadScreenMoviePlayer/Private/PreLoadMoviePlayerScreenBase.cpp
--- a/Engine/Plugins/Runtime/PreLoadScreenMoviePlayer/Source/PreLoadScreenMoviePlayer/Private/PreLoadMoviePlayerScreenBase.cpp
+++ b/Engine/Plugins/Runtime/PreLoadScreenMoviePlayer/Source/PreLoadScreenMoviePlayer/Private/PreLoadMoviePlayerScreenBase.cpp
@@ -80,7 +80,7 @@ protected:
     FOnKeyDown OnKeyDownHandler;
 };
 
-bool FPreLoadMoviePlayerScreenBase::Init()
+bool FPreLoadMoviePlayerScreenBase::Init(TFunctionRef<TSharedPtr<IAnalyticsProviderET>()> AnalyticsFactory)
 {
     if (!bInitialized)
     {
diff --git a/Engine/Plugins/Runtime/PreLoadScreenMoviePlayer/Source/PreLoadScreenMoviePlayer/Public/PreLoadMoviePlayerScreenBase.h b/Engine/Plugins/Runtime/PreLoadScreenMoviePlayer/Source/PreLoadScreenMoviePlayer/Public/PreLoadMoviePlayerScreenBase.h
--- a/Engine/Plugins/Runtime/PreLoadScreenMoviePlayer/Source/PreLoadScreenMoviePlayer/Public/PreLoadMoviePlayerScreenBase.h
+++ b/Engine/Plugins/Runtime/PreLoadScreenMoviePlayer/Source/PreLoadScreenMoviePlayer/Public/PreLoadMoviePlayerScreenBase.h
@@ -19,7 +19,7 @@ public:
     virtual void Tick(float DeltaTime) override;
     virtual void RenderTick(float DeltaTime) override;
 
-    virtual bool Init() override;
+    virtual bool Init(TFunctionRef<TSharedPtr<IAnalyticsProviderET>()> AnalyticsFactory) override;
 
     virtual TSharedPtr<SWidget> GetWidget() override { return MoviePlayerContents; }
     virtual const TSharedPtr<const SWidget> GetWidget() const override { return MoviePlayerContents; }
