#PreLoadScreenMoviePlayer CIS fix

diff --git a/Engine/Plugins/Runtime/PreLoadScreenMoviePlayer/Source/PreLoadScreenMoviePlayer/Private/PreLoadMoviePlayerModuleBase.cpp b/Engine/Plugins/Runtime/PreLoadScreenMoviePlayer/Source/PreLoadScreenMoviePlayer/Private/PreLoadMoviePlayerModuleBase.cpp
--- a/Engine/Plugins/Runtime/PreLoadScreenMoviePlayer/Source/PreLoadScreenMoviePlayer/Private/PreLoadMoviePlayerModuleBase.cpp
+++ b/Engine/Plugins/Runtime/PreLoadScreenMoviePlayer/Source/PreLoadScreenMoviePlayer/Private/PreLoadMoviePlayerModuleBase.cpp
@@ -13,7 +13,7 @@ void FPreLoadMoviePlayerScreenModuleBase::RegisterMovieStreamer(TSharedPtr<class
     {
         //Pass the streamer to our screen and then init with this new information
         MoviePreLoadScreen->RegisterMovieStreamer(InMovieStreamer);
-        MoviePreLoadScreen->Init();
+        MoviePreLoadScreen->Init([]()->TSharedPtr<IAnalyticsProviderET> { return TSharedPtr<IAnalyticsProviderET>(); });
     }
 }
 
