#PreLoadScreenMoviePlayer CIS fix

diff --git a/Engine/Plugins/Runtime/PreLoadScreenMoviePlayer/Source/PreLoadScreenMoviePlayer/Private/PreLoadMoviePlayerScreenBase.cpp b/Engine/Plugins/Runtime/PreLoadScreenMoviePlayer/Source/PreLoadScreenMoviePlayer/Private/PreLoadMoviePlayerScreenBase.cpp
--- a/Engine/Plugins/Runtime/PreLoadScreenMoviePlayer/Source/PreLoadScreenMoviePlayer/Private/PreLoadMoviePlayerScreenBase.cpp
+++ b/Engine/Plugins/Runtime/PreLoadScreenMoviePlayer/Source/PreLoadScreenMoviePlayer/Private/PreLoadMoviePlayerScreenBase.cpp
@@ -129,7 +129,7 @@ void FPreLoadMoviePlayerScreenBase::OnPlay(TWeakPtr<SWindow> TargetWindow)
     //Make sure Init has been called before we play a movie.
     if (!bInitialized)
     {
-        Init();
+        Init([]()->TSharedPtr<IAnalyticsProviderET>{ return TSharedPtr<IAnalyticsProviderET>(); });
     }
 
     LastPlayTime = FPlatformTime::Seconds();
