Fix non-unity
#rnx

diff --git a/Engine/Plugins/Runtime/PreLoadScreenMoviePlayer/Source/PreLoadScreenMoviePlayer/Private/PreLoadMoviePlayerScreenBase.cpp b/Engine/Plugins/Runtime/PreLoadScreenMoviePlayer/Source/PreLoadScreenMoviePlayer/Private/PreLoadMoviePlayerScreenBase.cpp
--- a/Engine/Plugins/Runtime/PreLoadScreenMoviePlayer/Source/PreLoadScreenMoviePlayer/Private/PreLoadMoviePlayerScreenBase.cpp
+++ b/Engine/Plugins/Runtime/PreLoadScreenMoviePlayer/Source/PreLoadScreenMoviePlayer/Private/PreLoadMoviePlayerScreenBase.cpp
@@ -7,6 +7,7 @@
 #include "HAL/FileManager.h"
 #include "Misc/ConfigCacheIni.h"
 
+#include "Framework/SlateDelegates.h"
 #include "Widgets/Layout/SDPIScaler.h"
 #include "Widgets/Layout/SBox.h"
 #include "Widgets/Layout/SBorder.h"
