#10802 tag

diff --git a/Engine/Plugins/Runtime/PreLoadScreenMoviePlayer/PreLoadScreenMoviePlayer.uplugin b/Engine/Plugins/Runtime/PreLoadScreenMoviePlayer/PreLoadScreenMoviePlayer.uplugin
--- a/Engine/Plugins/Runtime/PreLoadScreenMoviePlayer/PreLoadScreenMoviePlayer.uplugin
+++ b/Engine/Plugins/Runtime/PreLoadScreenMoviePlayer/PreLoadScreenMoviePlayer.uplugin
@@ -3,7 +3,7 @@
 	"Version" : 1,
 	"VersionName" : "1.0",
 	"FriendlyName" : "Pre-Load Screen Movie Player",
-	"Description" : "Handles a default implementation of using a Pre-Load screen to dislplay an engine loading movie",
+	"Description" : "Handles a default implementation of using a Pre-Load screen to display an engine loading movie.",
 	"Category" : "PreLoadScreenMoviePlayer",
 	"CreatedBy" : "Epic Games, Inc.",
 	"CreatedByURL" : "https://epicgames.com",
