Fixed Preload Movies not playing if Pre-Load Screen Movie Player plugin is enabled.

diff --git a/Engine/Plugins/Runtime/PreLoadScreenMoviePlayer/Source/PreLoadScreenMoviePlayer/Private/PreLoadMoviePlayerScreenBase.cpp b/Engine/Plugins/Runtime/PreLoadScreenMoviePlayer/Source/PreLoadScreenMoviePlayer/Private/PreLoadMoviePlayerScreenBase.cpp
--- a/Engine/Plugins/Runtime/PreLoadScreenMoviePlayer/Source/PreLoadScreenMoviePlayer/Private/PreLoadMoviePlayerScreenBase.cpp
+++ b/Engine/Plugins/Runtime/PreLoadScreenMoviePlayer/Source/PreLoadScreenMoviePlayer/Private/PreLoadMoviePlayerScreenBase.cpp
@@ -86,43 +86,53 @@ protected:
 
 void FPreLoadMoviePlayerScreenBase::Init()
 {
-    if (!bInitialized)
-    {
-        bInitialized = true;
-
-        MoviePlayerContents = SNew(SDefaultPreLoadScreenMovieBorder)
-		.OnKeyDown(this, &FPreLoadMoviePlayerScreenBase::OnLoadingScreenKeyDown)
-		.OnMouseButtonDown(this, &FPreLoadMoviePlayerScreenBase::OnLoadingScreenMouseButtonDown)
-		[
-            SNew(SOverlay)
-            +SOverlay::Slot()
-            .HAlign(HAlign_Center)
-            .VAlign(VAlign_Center)
+	if (bInitialized || !MovieStreamer.IsValid())
+	{
+		UE_LOG(LogPreloadMoviePlayer, Log, TEXT("Init failed for PreLoadMoviePlayerScreenBase"));
+		return;
+	}
+
+	InitSettingsFromConfig(GGameIni);
+	if (MovieAttributes.MoviePaths.IsEmpty())
+	{
+		UE_LOG(LogPreloadMoviePlayer, Log, TEXT("Init failed for PreLoadMoviePlayerScreenBase, no MoviePaths defined."));
+		return;
+	}
+
+    bInitialized = true;
+
+    MoviePlayerContents = SNew(SDefaultPreLoadScreenMovieBorder)
+	.OnKeyDown(this, &FPreLoadMoviePlayerScreenBase::OnLoadingScreenKeyDown)
+	.OnMouseButtonDown(this, &FPreLoadMoviePlayerScreenBase::OnLoadingScreenMouseButtonDown)
+	[
+        SNew(SOverlay)
+        +SOverlay::Slot()
+        .HAlign(HAlign_Center)
+        .VAlign(VAlign_Center)
+        [
+            SNew(SBox)
+            .WidthOverride(this, &FPreLoadMoviePlayerScreenBase::GetMovieWidth)
+            .HeightOverride(this, &FPreLoadMoviePlayerScreenBase::GetMovieHeight)
             [
-                SNew(SBox)
-                .WidthOverride(this, &FPreLoadMoviePlayerScreenBase::GetMovieWidth)
-                .HeightOverride(this, &FPreLoadMoviePlayerScreenBase::GetMovieHeight)
-                [
-                    SAssignNew(MovieViewport, SViewport)
-                    .EnableGammaCorrection(false)
-                    .Visibility(this, &FPreLoadMoviePlayerScreenBase::GetViewportVisibility)
-                ]
+                SAssignNew(MovieViewport, SViewport)
+                .EnableGammaCorrection(false)
+                .Visibility(this, &FPreLoadMoviePlayerScreenBase::GetViewportVisibility)
             ]
-		];
+		]
+	];
 
-        if (MovieStreamer.IsValid())
-        {
-            MovieViewport->SetViewportInterface(MovieStreamer->GetViewportInterface().ToSharedRef());
-        }
+	MovieStreamer->Init(MovieAttributes.MoviePaths, ((EMoviePlaybackType)((int)MovieAttributes.PlaybackType)));
 
-        MovieViewport->SetActive(true);
+	MovieViewport->SetViewportInterface(MovieStreamer->GetViewportInterface().ToSharedRef());
+    MovieViewport->SetActive(true);
 
-        // Register the movie viewport so that it can receive user input.
-        if (!FPlatformProperties::SupportsWindowedMode())
-        {
-            FSlateApplication::Get().RegisterGameViewport(MovieViewport.ToSharedRef());
-        }
+    // Register the movie viewport so that it can receive user input.
+    if (!FPlatformProperties::SupportsWindowedMode())
+    {
+        FSlateApplication::Get().RegisterGameViewport(MovieViewport.ToSharedRef());
     }
+
+	UE_LOG(LogPreloadMoviePlayer, Log, TEXT("Init completed for PreLoadMoviePlayerScreenBase"));
 }
 
 void FPreLoadMoviePlayerScreenBase::OnPlay(TWeakPtr<SWindow> TargetWindow)
@@ -151,6 +161,7 @@ void FPreLoadMoviePlayerScreenBase::CleanUp()
     }
 
     MovieStreamer.Reset();
+	MoviePlayerContents.Reset();
 }
 
 void FPreLoadMoviePlayerScreenBase::Tick(float DeltaTime)
@@ -265,55 +276,51 @@ bool FPreLoadMoviePlayerScreenBase::IsMovieStreamingFinished() const
 
 void FPreLoadMoviePlayerScreenBase::InitSettingsFromConfig(const FString& ConfigFileName)
 {
-    // We may have already setup a movie from a startup module
-    if (!MovieAttributes.IsValid())
-    {
-        // fill out the attributes
-        FPreLoadMovieAttributes MovieScreen;
+    // fill out the attributes
+    FPreLoadMovieAttributes MovieScreen;
 
-        bool bWaitForMoviesToComplete = false;
-        // Note: this code is executed too early so we cannot access UMoviePlayerSettings because the configs for that object have not been loaded and coalesced .  Have to read directly from the configs instead
-        GConfig->GetBool(TEXT("/Script/MoviePlayer.MoviePlayerSettings"), TEXT("bWaitForMoviesToComplete"), bWaitForMoviesToComplete, GGameIni);
-        GConfig->GetBool(TEXT("/Script/MoviePlayer.MoviePlayerSettings"), TEXT("bMoviesAreSkippable"), MovieScreen.bMoviesAreSkippable, GGameIni);
+    bool bWaitForMoviesToComplete = false;
+    // Note: this code is executed too early so we cannot access UMoviePlayerSettings because the configs for that object have not been loaded and coalesced .  Have to read directly from the configs instead
+    GConfig->GetBool(TEXT("/Script/MoviePlayer.MoviePlayerSettings"), TEXT("bWaitForMoviesToComplete"), bWaitForMoviesToComplete, GGameIni);
+    GConfig->GetBool(TEXT("/Script/MoviePlayer.MoviePlayerSettings"), TEXT("bMoviesAreSkippable"), MovieScreen.bMoviesAreSkippable, GGameIni);
 
-        MovieScreen.bAutoCompleteWhenLoadingCompletes = !bWaitForMoviesToComplete;
+    MovieScreen.bAutoCompleteWhenLoadingCompletes = !bWaitForMoviesToComplete;
 
-        TArray<FString> StartupMovies;
-        GConfig->GetArray(TEXT("/Script/MoviePlayer.MoviePlayerSettings"), TEXT("StartupMovies"), StartupMovies, GGameIni);
+    TArray<FString> StartupMovies;
+    GConfig->GetArray(TEXT("/Script/MoviePlayer.MoviePlayerSettings"), TEXT("StartupMovies"), StartupMovies, GGameIni);
 
-        if (StartupMovies.Num() == 0)
-        {
-            StartupMovies.Add(TEXT("Default_Startup"));
-        }
+    if (StartupMovies.Num() == 0)
+    {
+        StartupMovies.Add(TEXT("Default_Startup"));
+    }
 
-        // double check that the movies exist
-        // We dont know the extension so compare against any file in the directory with the same name for now
-        // @todo New Movie Player: movies should have the extension on them when set via the project settings
-        TArray<FString> ExistingMovieFiles;
-        IFileManager::Get().FindFiles(ExistingMovieFiles, *(FPaths::ProjectContentDir() + TEXT("Movies")));
+    // double check that the movies exist
+    // We dont know the extension so compare against any file in the directory with the same name for now
+    // @todo New Movie Player: movies should have the extension on them when set via the project settings
+    TArray<FString> ExistingMovieFiles;
+    IFileManager::Get().FindFiles(ExistingMovieFiles, *(FPaths::ProjectContentDir() + TEXT("Movies")));
 
-        bool bHasValidMovie = false;
-        for (const FString& Movie : StartupMovies)
+    bool bHasValidMovie = false;
+    for (const FString& Movie : StartupMovies)
+    {
+        bool bFound = ExistingMovieFiles.ContainsByPredicate(
+            [&Movie](const FString& ExistingMovie)
         {
-            bool bFound = ExistingMovieFiles.ContainsByPredicate(
-                [&Movie](const FString& ExistingMovie)
-            {
-                return ExistingMovie.Contains(Movie);
-            });
-
-            if (bFound)
-            {
-                bHasValidMovie = true;
-                MovieScreen.MoviePaths.Add(Movie);
-            }
-        }
+            return ExistingMovie.Contains(Movie);
+        });
 
-        if (bHasValidMovie)
+        if (bFound)
         {
-            // now setup the actual loading screen
-            MovieAttributes = MovieScreen;
+            bHasValidMovie = true;
+            MovieScreen.MoviePaths.Add(Movie);
         }
     }
+
+    if (bHasValidMovie)
+    {
+        // now setup the actual loading screen
+        MovieAttributes = MovieScreen;
+    }
 }
 
 void FPreLoadMoviePlayerScreenBase::RegisterMovieStreamer(TSharedPtr<IMovieStreamer, ESPMode::ThreadSafe> MovieStreamerIn)
diff --git a/Engine/Plugins/Runtime/PreLoadScreenMoviePlayer/Source/PreLoadScreenMoviePlayer/Public/PreLoadMoviePlayerScreenBase.h b/Engine/Plugins/Runtime/PreLoadScreenMoviePlayer/Source/PreLoadScreenMoviePlayer/Public/PreLoadMoviePlayerScreenBase.h
--- a/Engine/Plugins/Runtime/PreLoadScreenMoviePlayer/Source/PreLoadScreenMoviePlayer/Public/PreLoadMoviePlayerScreenBase.h
+++ b/Engine/Plugins/Runtime/PreLoadScreenMoviePlayer/Source/PreLoadScreenMoviePlayer/Public/PreLoadMoviePlayerScreenBase.h
@@ -8,6 +8,8 @@
 #include "Widgets/SWindow.h"
 #include "Widgets/SViewport.h"
 
+DEFINE_LOG_CATEGORY_STATIC(LogPreloadMoviePlayer, Log, All);
+
 class SViewport;
 
 class FPreLoadMoviePlayerScreenBase : public FPreLoadScreenBase
