UE-248290 - Ensuring skipping logic is respected in PreLoadScreenMoviePlayer plugin, matching the logic for DefaultGameMoviePlayer

diff --git a/Engine/Plugins/Runtime/PreLoadScreenMoviePlayer/Source/PreLoadScreenMoviePlayer/Private/PreLoadMoviePlayerScreenBase.cpp b/Engine/Plugins/Runtime/PreLoadScreenMoviePlayer/Source/PreLoadScreenMoviePlayer/Private/PreLoadMoviePlayerScreenBase.cpp
--- a/Engine/Plugins/Runtime/PreLoadScreenMoviePlayer/Source/PreLoadScreenMoviePlayer/Private/PreLoadMoviePlayerScreenBase.cpp
+++ b/Engine/Plugins/Runtime/PreLoadScreenMoviePlayer/Source/PreLoadScreenMoviePlayer/Private/PreLoadMoviePlayerScreenBase.cpp
@@ -3,18 +3,18 @@
 #include "PreLoadMoviePlayerScreenBase.h"
 #include "PreLoadScreenManager.h"
 #include "PreLoadMoviePlayerModule.h"
-
 #include "Framework/Application/SlateApplication.h"
 #include "Framework/SlateDelegates.h"
 #include "HAL/FileManager.h"
 #include "Misc/ConfigCacheIni.h"
-
 #include "Framework/Application/SlateApplication.h"
 #include "Widgets/Layout/SDPIScaler.h"
 #include "Widgets/Layout/SBox.h"
 #include "Widgets/Layout/SBorder.h"
 #include "Widgets/SViewport.h"
 #include "Widgets/DeclarativeSyntaxSupport.h"
+#include "GlobalShader.h"
+#include "ShaderCompiler.h"
 
 class SDefaultPreLoadScreenMovieBorder : public SBorder
 {
@@ -99,38 +99,46 @@ void FPreLoadMoviePlayerScreenBase::Init()
 		return;
 	}
 
-    bInitialized = true;
-
-    MoviePlayerContents = SNew(SDefaultPreLoadScreenMovieBorder)
-	.OnKeyDown(this, &FPreLoadMoviePlayerScreenBase::OnLoadingScreenKeyDown)
-	.OnMouseButtonDown(this, &FPreLoadMoviePlayerScreenBase::OnLoadingScreenMouseButtonDown)
-	[
-        SNew(SOverlay)
-        +SOverlay::Slot()
-        .HAlign(HAlign_Center)
-        .VAlign(VAlign_Center)
-        [
-            SNew(SBox)
-            .WidthOverride(this, &FPreLoadMoviePlayerScreenBase::GetMovieWidth)
-            .HeightOverride(this, &FPreLoadMoviePlayerScreenBase::GetMovieHeight)
-            [
-                SAssignNew(MovieViewport, SViewport)
-                .EnableGammaCorrection(false)
-                .Visibility(this, &FPreLoadMoviePlayerScreenBase::GetViewportVisibility)
-            ]
-		]
-	];
+	// Initialize shaders, because otherwise they might not be guaranteed to exist at this point
+	if (!FPlatformProperties::RequiresCookedData())
+	{
+		TArray<int32> ShaderMapIds;
+		ShaderMapIds.Add(GlobalShaderMapId);
+		GShaderCompilingManager->FinishCompilation(TEXT("Global"), ShaderMapIds);
+	}
+
+	bInitialized = true;
+
+	MoviePlayerContents = SNew(SDefaultPreLoadScreenMovieBorder)
+		.OnKeyDown(this, &FPreLoadMoviePlayerScreenBase::OnLoadingScreenKeyDown)
+		.OnMouseButtonDown(this, &FPreLoadMoviePlayerScreenBase::OnLoadingScreenMouseButtonDown)
+		[
+			SNew(SOverlay)
+				+ SOverlay::Slot()
+				.HAlign(HAlign_Center)
+				.VAlign(VAlign_Center)
+				[
+					SNew(SBox)
+						.WidthOverride(this, &FPreLoadMoviePlayerScreenBase::GetMovieWidth)
+						.HeightOverride(this, &FPreLoadMoviePlayerScreenBase::GetMovieHeight)
+						[
+							SAssignNew(MovieViewport, SViewport)
+								.EnableGammaCorrection(false)
+								.Visibility(this, &FPreLoadMoviePlayerScreenBase::GetViewportVisibility)
+						]
+				]
+		];
 
 	MovieStreamer->Init(MovieAttributes.MoviePaths, ((EMoviePlaybackType)((int)MovieAttributes.PlaybackType)));
 
 	MovieViewport->SetViewportInterface(MovieStreamer->GetViewportInterface().ToSharedRef());
-    MovieViewport->SetActive(true);
+	MovieViewport->SetActive(true);
 
-    // Register the movie viewport so that it can receive user input.
-    if (!FPlatformProperties::SupportsWindowedMode())
-    {
-        FSlateApplication::Get().RegisterGameViewport(MovieViewport.ToSharedRef());
-    }
+	// Register the movie viewport so that it can receive user input.
+	if (!FPlatformProperties::SupportsWindowedMode())
+	{
+		FSlateApplication::Get().RegisterGameViewport(MovieViewport.ToSharedRef());
+	}
 
 	UE_LOG(LogPreloadMoviePlayer, Log, TEXT("Init completed for PreLoadMoviePlayerScreenBase"));
 }
@@ -144,7 +152,7 @@ void FPreLoadMoviePlayerScreenBase::OnPlay(TWeakPtr<SWindow> TargetWindow)
         Init();
     }
 
-    LastPlayTime = FPlatformTime::Seconds();
+    LastPlayTime = FPlatformTime::Seconds();	
 }
 
 void FPreLoadMoviePlayerScreenBase::OnStop()
@@ -166,6 +174,12 @@ void FPreLoadMoviePlayerScreenBase::CleanUp()
 
 void FPreLoadMoviePlayerScreenBase::Tick(float DeltaTime)
 {
+	// Make sure if focus is lost for some reason it is restored to allow skipping
+	if (MoviePlayerContents.IsValid() && !MoviePlayerContents->HasAnyUserFocus())
+	{
+		FSlateApplication::Get().SetAllUserFocus(MoviePlayerContents);
+	}
+
     const bool bAutoCompleteWhenLoadingCompletes = MovieAttributes.bAutoCompleteWhenLoadingCompletes;
     const bool bWaitForManualStop = MovieAttributes.bWaitForManualStop;
     const bool bEnforceMinimumTime = MovieAttributes.MinimumLoadingScreenDisplayTime >= 0.0f;
@@ -186,11 +200,15 @@ void FPreLoadMoviePlayerScreenBase::RenderTick(float DeltaTime)
 {
     if (MovieStreamingIsPrepared())
     {
+		MovieStreamer->TickPreEngine();
+		MovieStreamer->TickPostEngine();
+
         const bool bMovieIsDone = MovieStreamer->Tick(DeltaTime);
         if (bMovieIsDone)
         {
             MovieStreamingIsDone.Set(1);
         }
+		MovieStreamer->TickPostRender();
     }
 }
 
diff --git a/Engine/Source/Runtime/PreLoadScreen/Private/PreLoadScreenManager.cpp b/Engine/Source/Runtime/PreLoadScreen/Private/PreLoadScreenManager.cpp
--- a/Engine/Source/Runtime/PreLoadScreen/Private/PreLoadScreenManager.cpp
+++ b/Engine/Source/Runtime/PreLoadScreen/Private/PreLoadScreenManager.cpp
@@ -664,6 +664,17 @@ void FPreLoadScreenManager::WaitForEngineLoadingScreenToFinish()
 	if (HasActivePreLoadScreenType(EPreLoadScreenTypes::EngineLoadingScreen))
 	{
 		IPreLoadScreen* ActivePreLoadScreen = GetActivePreLoadScreen();
+
+		//Add the current screens widget to focus so that any input gets funneled to it correctly.
+		if (ActivePreLoadScreen)
+		{
+			TSharedPtr<SWidget> Widget = ActivePreLoadScreen->GetWidget();
+			if (Widget.IsValid())
+			{
+				FSlateApplication::Get().SetAllUserFocus(Widget);
+			}
+		}
+
 		while (ActivePreLoadScreen && !ActivePreLoadScreen->IsDone())
 		{
 			GameLogicFrameTick();
